Added UdpMessaging to ProgramEnabledPlugins list in their .ini files

diff --git a/Engine/Programs/UnrealFrontend/Config/DefaultEngine.ini b/Engine/Programs/UnrealFrontend/Config/DefaultEngine.ini
--- a/Engine/Programs/UnrealFrontend/Config/DefaultEngine.ini
+++ b/Engine/Programs/UnrealFrontend/Config/DefaultEngine.ini
@@ -8,3 +8,4 @@ GameNameShort=UFE
 [Plugins]
 +EnabledPlugins="MessagingDebugger"
 -EnabledPlugins="ExampleDeviceProfileSelector"
++ProgramEnabledPlugins="UdpMessaging"
diff --git a/Engine/Programs/UnrealLightmass/Config/DefaultEngine.ini b/Engine/Programs/UnrealLightmass/Config/DefaultEngine.ini
--- a/Engine/Programs/UnrealLightmass/Config/DefaultEngine.ini
+++ b/Engine/Programs/UnrealLightmass/Config/DefaultEngine.ini
@@ -3,3 +3,4 @@ BasedOn=..\Engine\Config\BaseEngine.ini
 
 [Plugins]
 -EnabledPlugins="ExampleDeviceProfileSelector"
++ProgramEnabledPlugins="UdpMessaging"
