fix for failure to run UFE

diff --git a/Engine/Programs/UnrealFrontend/Config/DefaultEngine.ini b/Engine/Programs/UnrealFrontend/Config/DefaultEngine.ini
--- a/Engine/Programs/UnrealFrontend/Config/DefaultEngine.ini
+++ b/Engine/Programs/UnrealFrontend/Config/DefaultEngine.ini
@@ -5,3 +5,5 @@ GameNameShort=UFE
 +EnabledPlugins="MessagingDebugger"
 -EnabledPlugins="ExampleDeviceProfileSelector"
 +ProgramEnabledPlugins="UdpMessaging"
++ProgramEnabledPlugins="VisualStudioSourceCodeAccess"
++ProgramEnabledPlugins="XcodeSourceCodeAccess"
