[Trivial] Fix for Project Launcher not saving settings on exit.

diff --git a/Engine/Programs/UnrealFrontend/Config/DefaultEngine.ini b/Engine/Programs/UnrealFrontend/Config/DefaultEngine.ini
--- a/Engine/Programs/UnrealFrontend/Config/DefaultEngine.ini
+++ b/Engine/Programs/UnrealFrontend/Config/DefaultEngine.ini
@@ -8,3 +8,8 @@ GameNameShort=UFE
 +ProgramEnabledPlugins="TcpMessaging"
 +ProgramEnabledPlugins="VisualStudioSourceCodeAccess"
 +ProgramEnabledPlugins="XcodeSourceCodeAccess"
+
+[SectionsToSave]
+bCanSaveAllSections=true
+Section=CurrentIniVersion
+
